Fix for Linux trying to build MagicLeapAnalytics by swapping from Blacklist to Whitelist

diff --git a/Engine/Plugins/Lumin/MagicLeapAnalytics/MagicLeapAnalytics.uplugin b/Engine/Plugins/Lumin/MagicLeapAnalytics/MagicLeapAnalytics.uplugin
--- a/Engine/Plugins/Lumin/MagicLeapAnalytics/MagicLeapAnalytics.uplugin
+++ b/Engine/Plugins/Lumin/MagicLeapAnalytics/MagicLeapAnalytics.uplugin
@@ -23,7 +23,7 @@
 			"Name" : "MagicLeapAnalytics",
 			"Type" : "Runtime",
 			"LoadingPhase" : "PreDefault",
-			"BlacklistPlatforms" :
+			"WhitelistPlatforms" :
 			[
 				"Win64",
 				"Mac",
