Wrapping Lumin call with PLATFORM_LUMIN

diff --git a/Engine/Plugins/Lumin/MagicLeapAnalytics/Source/MagicLeapAnalytics/Private/MagicLeapAnalytics.cpp b/Engine/Plugins/Lumin/MagicLeapAnalytics/Source/MagicLeapAnalytics/Private/MagicLeapAnalytics.cpp
--- a/Engine/Plugins/Lumin/MagicLeapAnalytics/Source/MagicLeapAnalytics/Private/MagicLeapAnalytics.cpp
+++ b/Engine/Plugins/Lumin/MagicLeapAnalytics/Source/MagicLeapAnalytics/Private/MagicLeapAnalytics.cpp
@@ -117,6 +117,7 @@ void FMagicLeapAnalyticsProvider::EndSession()
 		TSharedRef<TJsonWriter<TCHAR>> Writer = TJsonWriterFactory<>::Create(&JsonStr);
 		FJsonSerializer::Serialize<TCHAR>(LogJson.ToSharedRef(), Writer, true);
 
+#if PLATFORM_LUMIN
 		MLAppAnalyticsError error;
 		bool bResult = MLAppAnalyticsUpload(TCHAR_TO_UTF8(*SessionId), TCHAR_TO_UTF8(*JsonStr), &error);
 		if (bResult)
@@ -127,6 +128,7 @@ void FMagicLeapAnalyticsProvider::EndSession()
 		{
 			UE_LOG(LogMagicLeapAnalytics, Error, TEXT("Error saving analytics. Error code = %d"), static_cast<int32>(error));
 		}
+#endif
 
 #if !PLATFORM_LUMIN
 		if (FileArchive != nullptr)
