Fix for Lumin Linux dependency in a better way

diff --git a/Engine/Plugins/Lumin/MagicLeapAnalytics/MagicLeapAnalytics.uplugin b/Engine/Plugins/Lumin/MagicLeapAnalytics/MagicLeapAnalytics.uplugin
--- a/Engine/Plugins/Lumin/MagicLeapAnalytics/MagicLeapAnalytics.uplugin
+++ b/Engine/Plugins/Lumin/MagicLeapAnalytics/MagicLeapAnalytics.uplugin
@@ -23,11 +23,12 @@
 			"Name" : "MagicLeapAnalytics",
 			"Type" : "Runtime",
 			"LoadingPhase" : "PreDefault",
-			"WhitelistPlatforms" :
+			"BlacklistPlatforms" :
 			[
 				"Win64",
 				"Mac",
-				"Lumin"
+				"Lumin",
+				"Linux"
 			]			
 		}
 	],
